Fix a crash when provided test exclusion Reason string is larger than 1024 characters.

#rnx

diff --git a/Engine/Source/Runtime/AutomationTest/Public/AutomationTestExcludelist.h b/Engine/Source/Runtime/AutomationTest/Public/AutomationTestExcludelist.h
--- a/Engine/Source/Runtime/AutomationTest/Public/AutomationTestExcludelist.h
+++ b/Engine/Source/Runtime/AutomationTest/Public/AutomationTestExcludelist.h
@@ -45,7 +45,7 @@ struct FAutomationTestExcludeOptions
 
 	/* Reason to why the test is excluded */
 	UPROPERTY(EditAnywhere, Category = ExcludeTestOptions)
-	FString Reason;
+	FName Reason;
 
 	/* Option to target specific RHI. No option means it should be applied to all RHI */
 	UPROPERTY(EditAnywhere, Meta = (Bitmask, BitmaskEnum = "ERHI_Flags"), Category = ExcludeTestOptions)
@@ -87,7 +87,7 @@ struct FAutomationTestExcludelistEntry
 	{
 		TSharedPtr<FAutomationTestExcludeOptions> Options = MakeShareable(new FAutomationTestExcludeOptions());
 		Options->Test = Test;
-		Options->Reason = Reason.ToString();
+		Options->Reason = Reason;
 		Options->Warn = Warn;
 		Options->SetRHIFlagsFromNames(RHIs);
 		
