Checkout config file before updating test exclude list.

#rnx

diff --git a/Engine/Source/Runtime/AutomationTest/AutomationTest.Build.cs b/Engine/Source/Runtime/AutomationTest/AutomationTest.Build.cs
--- a/Engine/Source/Runtime/AutomationTest/AutomationTest.Build.cs
+++ b/Engine/Source/Runtime/AutomationTest/AutomationTest.Build.cs
@@ -17,6 +17,15 @@ namespace UnrealBuildTool.Rules
 				{
 					"Core",
 				});
+
+			if (Target.bBuildEditor)
+			{
+				PrivateDependencyModuleNames.AddRange(
+					new string[] {
+						"SourceControl",
+						"TargetPlatform",
+					});
+			}
 		}
 	}
 }
diff --git a/Engine/Source/Runtime/AutomationTest/Private/AutomationTestExcludelist.cpp b/Engine/Source/Runtime/AutomationTest/Private/AutomationTestExcludelist.cpp
--- a/Engine/Source/Runtime/AutomationTest/Private/AutomationTestExcludelist.cpp
+++ b/Engine/Source/Runtime/AutomationTest/Private/AutomationTestExcludelist.cpp
@@ -2,6 +2,16 @@
 
 #include "AutomationTestExcludelist.h"
 
+#if WITH_EDITOR
+	#include "HAL/PlatformFileManager.h"
+	#include "ISourceControlOperation.h"
+	#include "SourceControlOperations.h"
+	#include "ISourceControlProvider.h"
+	#include "ISourceControlModule.h"
+#endif
+
+DEFINE_LOG_CATEGORY_STATIC(LogAutomationTestExcludelist, Log, All);
+
 static const FString FunctionalTestsPreFix = TEXT("Project.Functional Tests.");
 
 void UAutomationTestExcludelist::OverrideConfigSection(FString& SectionName)
@@ -121,3 +131,105 @@ FAutomationTestExcludelistEntry* UAutomationTestExcludelist::GetExcludeTestEntry
 
 	return nullptr;
 }
+
+#if WITH_EDITOR
+bool CheckOutOrAddFile(const FString& InFileToCheckOut)
+{
+	bool bSuccessfullyCheckedOutOrAddedFile = false;
+	if (ISourceControlModule::Get().IsEnabled())
+	{
+		ISourceControlProvider& SourceControlProvider = ISourceControlModule::Get().GetProvider();
+		FSourceControlStatePtr SourceControlState = SourceControlProvider.GetState(InFileToCheckOut, EStateCacheUsage::Use);
+
+		TArray<FString> FilesToBeCheckedOut;
+		FilesToBeCheckedOut.Add(InFileToCheckOut);
+
+		if (SourceControlState.IsValid())
+		{
+			if (SourceControlState->IsSourceControlled())
+			{
+				if (SourceControlState->IsDeleted())
+				{
+					UE_LOG(LogAutomationTestExcludelist, Error, TEXT("The configuration file is marked for deletion."));
+				}
+				else if (SourceControlState->CanCheckout() || SourceControlState->IsCheckedOutOther() || FPlatformFileManager::Get().GetPlatformFile().IsReadOnly(*InFileToCheckOut))
+				{
+					ECommandResult::Type CommandResult = SourceControlProvider.Execute(ISourceControlOperation::Create<FCheckOut>(), FilesToBeCheckedOut);
+					if (CommandResult == ECommandResult::Failed)
+					{
+						UE_LOG(LogAutomationTestExcludelist, Error, TEXT("Failed to check out the configuration file."));
+					}
+					else if (CommandResult == ECommandResult::Cancelled)
+					{
+						UE_LOG(LogAutomationTestExcludelist, Warning, TEXT("Checkout was cancelled."));
+					}
+					else
+					{
+						bSuccessfullyCheckedOutOrAddedFile = true;
+					}
+				}
+			}
+			else if (!SourceControlState->IsUnknown())
+			{
+				if (!FPlatformFileManager::Get().GetPlatformFile().FileExists(*InFileToCheckOut))
+				{
+					return true;
+				}
+
+				ECommandResult::Type CommandResult = SourceControlProvider.Execute(ISourceControlOperation::Create<FMarkForAdd>(), FilesToBeCheckedOut);
+
+				if (CommandResult == ECommandResult::Failed)
+				{
+					UE_LOG(LogAutomationTestExcludelist, Error, TEXT("Failed to check out the configuration file."));
+				}
+				else if (CommandResult == ECommandResult::Cancelled)
+				{
+					UE_LOG(LogAutomationTestExcludelist, Warning, TEXT("Checkout was cancelled.."));
+				}
+				else
+				{
+					bSuccessfullyCheckedOutOrAddedFile = true;
+				}
+			}
+		}
+	}
+	return bSuccessfullyCheckedOutOrAddedFile;
+}
+
+bool MakeWritable(const FString& InFileToMakeWritable)
+{
+	if (!FPlatformFileManager::Get().GetPlatformFile().FileExists(*InFileToMakeWritable))
+	{
+		return true;
+	}
+
+	return FPlatformFileManager::Get().GetPlatformFile().SetReadOnly(*InFileToMakeWritable, false);
+}
+#endif
+
+void UAutomationTestExcludelist::SaveConfig()
+{
+#if WITH_EDITOR
+	FString ConfigFilename = GetConfigFilename();
+	bool bIsWritable = FPlatformFileManager::Get().GetPlatformFile().FileExists(*ConfigFilename) && !FPlatformFileManager::Get().GetPlatformFile().IsReadOnly(*ConfigFilename);
+	if (!bIsWritable)
+	{
+		bIsWritable = CheckOutOrAddFile(ConfigFilename);
+		if (!bIsWritable)
+		{
+			UE_LOG(LogAutomationTestExcludelist, Warning, TEXT("Config file '%s' is readonly and could not be checked out. File will be marked writable."), *ConfigFilename);
+			bIsWritable = MakeWritable(ConfigFilename);
+		}
+	}
+
+	if (!bIsWritable)
+	{
+		UE_LOG(LogAutomationTestExcludelist, Error, TEXT("Failed to make the configuration file '%s' writable."), *ConfigFilename);
+	}
+	else
+#endif
+	{
+		UObject::UpdateDefaultConfigFile();
+	}
+}
+
diff --git a/Engine/Source/Runtime/AutomationTest/Public/AutomationTestExcludelist.h b/Engine/Source/Runtime/AutomationTest/Public/AutomationTestExcludelist.h
--- a/Engine/Source/Runtime/AutomationTest/Public/AutomationTestExcludelist.h
+++ b/Engine/Source/Runtime/AutomationTest/Public/AutomationTestExcludelist.h
@@ -168,7 +168,8 @@ public:
 	bool IsTestExcluded(const FString& TestName, const FString& RHI = TEXT(""), FName* OutReason = nullptr, bool* OutWarn = nullptr);
 	FAutomationTestExcludelistEntry* GetExcludeTestEntry(const FString& TestName);
 
-	void SaveConfig() { UObject::UpdateDefaultConfigFile(); }
+	void SaveConfig();
+	FString GetConfigFilename() { return UObject::GetDefaultConfigFilename(); } const
 	// It is called automatically when CDO is created, usually you don't need to call LoadConfig manually
 	void LoadConfig() { UObject::LoadConfig(GetClass()); }
 
