Fix exclusion list not working when test name is input in the code with extra spaces.

diff --git a/Engine/Source/Runtime/AutomationTest/Private/AutomationTestExcludelist.cpp b/Engine/Source/Runtime/AutomationTest/Private/AutomationTestExcludelist.cpp
--- a/Engine/Source/Runtime/AutomationTest/Private/AutomationTestExcludelist.cpp
+++ b/Engine/Source/Runtime/AutomationTest/Private/AutomationTestExcludelist.cpp
@@ -59,7 +59,7 @@ FString UAutomationTestExcludelist::GetFullTestName(const FAutomationTestExclude
 void UAutomationTestExcludelist::AddToExcludeTest(const FString& TestName, const FAutomationTestExcludelistEntry& ExcludelistEntry)
 {
 	auto NewEntry = FAutomationTestExcludelistEntry(ExcludelistEntry);
-	NewEntry.Test = *TestName;
+	NewEntry.Test = *(TestName.TrimStartAndEnd());
 	if (!NewEntry.Map.IsNone())
 	{
 		NewEntry.Map = TEXT("");
@@ -75,7 +75,7 @@ void UAutomationTestExcludelist::RemoveFromExcludeTest(const FString& TestName)
 	if (TestName.IsEmpty())
 		return;
 
-	const FString NameToCompare = TestName.ToLower();
+	const FString NameToCompare = TestName.TrimStartAndEnd().ToLower();
 
 	for (int i = 0; i < ExcludeTest.Num(); ++i)
 	{
@@ -117,7 +117,7 @@ FAutomationTestExcludelistEntry* UAutomationTestExcludelist::GetExcludeTestEntry
 	if (TestName.IsEmpty())
 		return nullptr;
 
-	const FString NameToCompare = TestName.ToLower();
+	const FString NameToCompare = TestName.TrimStartAndEnd().ToLower();
 
 	for (auto& Entry : ExcludeTest)
 	{
diff --git a/Engine/Source/Runtime/AutomationTest/Public/AutomationTestExcludelist.h b/Engine/Source/Runtime/AutomationTest/Public/AutomationTestExcludelist.h
--- a/Engine/Source/Runtime/AutomationTest/Public/AutomationTestExcludelist.h
+++ b/Engine/Source/Runtime/AutomationTest/Public/AutomationTestExcludelist.h
@@ -97,7 +97,7 @@ struct FAutomationTestExcludelistEntry
 	/* Determine if exclusion entry is propagated based on test name - used for management in test automation window */
 	void SetPropagation(const FString& ForTestName)
 	{
-		bIsPropagated = FullTestName != ForTestName.ToLower();
+		bIsPropagated = FullTestName != ForTestName.TrimStartAndEnd().ToLower();
 	}
 
 	/* Return true if the entry is not specific */
