Add Metal in RHI options for Automation Test Exclude list

diff --git a/Engine/Source/Runtime/AutomationTest/Public/AutomationTestExcludelist.h b/Engine/Source/Runtime/AutomationTest/Public/AutomationTestExcludelist.h
--- a/Engine/Source/Runtime/AutomationTest/Public/AutomationTestExcludelist.h
+++ b/Engine/Source/Runtime/AutomationTest/Public/AutomationTestExcludelist.h
@@ -11,6 +11,7 @@ enum class ERHI_Flags : uint8
 	DirectX11	= 0x01 UMETA(DisplayName = "DirectX 11"),
 	DirectX12	= 0x02 UMETA(DisplayName = "DirectX 12"),
 	Vulkan		= 0x04 UMETA(DisplayName = "Vulkan"),
+	Metal		= 0x08 UMETA(DisplayName = "Metal"),
 	NUM			UMETA(Hidden)
 };
 
