Automation Test - Fix Previously Excluded Test not inheriting exclusion settings from the parent suite

diff --git a/Engine/Source/Runtime/AutomationTest/Private/AutomationTestExcludelist.cpp b/Engine/Source/Runtime/AutomationTest/Private/AutomationTestExcludelist.cpp
--- a/Engine/Source/Runtime/AutomationTest/Private/AutomationTestExcludelist.cpp
+++ b/Engine/Source/Runtime/AutomationTest/Private/AutomationTestExcludelist.cpp
@@ -25,6 +25,16 @@ UAutomationTestExcludelist* UAutomationTestExcludelist::Get()
 	return GetMutableDefault<UAutomationTestExcludelist>();
 }
 
+void SortExcludelist(TArray<FAutomationTestExcludelistEntry>& List)
+{
+	// Sort in alphabetical order, shortest to longest of FullTestName property.
+	// That is to naturally gives priority to parent suite over invidual test exclusion when calling GetExcludeTestEntry(TestName).
+	List.Sort([](const FAutomationTestExcludelistEntry& A, const FAutomationTestExcludelistEntry& B)
+	{
+		return A.FullTestName < B.FullTestName;
+	});
+}
+
 void UAutomationTestExcludelist::PostInitProperties()
 {
 	Super::PostInitProperties();
@@ -36,6 +46,8 @@ void UAutomationTestExcludelist::PostInitProperties()
 			Entry.FullTestName = GetFullTestName(Entry);
 		}
 	}
+
+	SortExcludelist(ExcludeTest);
 }
 
 FString UAutomationTestExcludelist::GetFullTestName(const FAutomationTestExcludelistEntry& ExcludelistEntry)
@@ -68,6 +80,7 @@ void UAutomationTestExcludelist::AddToExcludeTest(const FString& TestName, const
 	NewEntry.FullTestName = GetFullTestName(NewEntry);
 
 	ExcludeTest.Add(NewEntry);
+	SortExcludelist(ExcludeTest);
 }
 
 void UAutomationTestExcludelist::RemoveFromExcludeTest(const FString& TestName)
