Use dot syntax to name blueprint functional tests instead of using the full package path (replacing / by .)
+ support backcomp from functional test exclude list
+ support exclude test fallback when RHI options are specified

diff --git a/Engine/Plugins/Tests/EditorTests/Source/EditorTests/Private/UnrealEd/PlayMapInPIETests.cpp b/Engine/Plugins/Tests/EditorTests/Source/EditorTests/Private/UnrealEd/PlayMapInPIETests.cpp
--- a/Engine/Plugins/Tests/EditorTests/Source/EditorTests/Private/UnrealEd/PlayMapInPIETests.cpp
+++ b/Engine/Plugins/Tests/EditorTests/Source/EditorTests/Private/UnrealEd/PlayMapInPIETests.cpp
@@ -117,7 +117,7 @@ public:
 
 						FString MapAssetPath = MapAsset.ObjectPath.ToString();
 
-						OutBeautifiedNames.Add(MapPackageName.RightChop(6)); // Remove "/Game/" from the name
+						OutBeautifiedNames.Add(MapPackageName.RightChop(6).Replace(TEXT("/"), TEXT("."))); // Remove "/Game/" from the name and use dot syntax
 						OutTestCommands.Add(MapAssetPath + TEXT(";") + MapPackageName);
 					}
 				}
diff --git a/Engine/Source/Developer/FunctionalTesting/Private/FunctionalTestingModule.cpp b/Engine/Source/Developer/FunctionalTesting/Private/FunctionalTestingModule.cpp
--- a/Engine/Source/Developer/FunctionalTesting/Private/FunctionalTestingModule.cpp
+++ b/Engine/Source/Developer/FunctionalTesting/Private/FunctionalTestingModule.cpp
@@ -128,6 +128,7 @@ void FFunctionalTestingModule::GetMapTests(bool bEditorOnlyTests, TArray<FString
 				FString MapAssetPath = MapAsset.ObjectPath.ToString();
 				FString MapPackageName = MapAsset.PackageName.ToString();
 				if (!IsDeveloperDirectoryIncluded && MapPackageName.Find(TEXT("/Game/Developers")) == 0) continue;
+				FString PartialSuiteName = MapPackageName.RightChop(6).Replace(TEXT("/"), TEXT(".")); // Remove "/Game/" from the name and use dot syntax
 
 				FString AllTestNames;
 				FAssetDataTagMapSharedView::FFindTagResult MapTestNames = MapAsset.TagsAndValues.FindTag(bEditorOnlyTests ? TEXT("TestNamesEditor") : TEXT("TestNames"));
@@ -182,8 +183,8 @@ void FFunctionalTestingModule::GetMapTests(bool bEditorOnlyTests, TArray<FString
 
 						if (MapTest.Split(TEXT("|"), &BeautifulTestName, &RealTestName))
 						{
-							OutBeautifiedNames.Add(MapPackageName + TEXT(".") + *BeautifulTestName);
-							OutTestCommands.Add(MapAssetPath + TEXT(";") + MapAsset.PackageName.ToString() + TEXT(";") + *RealTestName);
+							OutBeautifiedNames.Add(PartialSuiteName + TEXT(".") + *BeautifulTestName);
+							OutTestCommands.Add(MapAssetPath + TEXT(";") + MapPackageName + TEXT(";") + *RealTestName);
 							OutTestMapAssets.AddUnique(MapAssetPath);
 						}
 					}
@@ -191,7 +192,7 @@ void FFunctionalTestingModule::GetMapTests(bool bEditorOnlyTests, TArray<FString
 				else if (!bEditorOnlyTests && MapAsset.AssetName.ToString().Find(TEXT("FTEST_")) == 0)
 				{
 					OutBeautifiedNames.Add(MapAsset.AssetName.ToString());
-					OutTestCommands.Add(MapAssetPath + TEXT(";") + MapAsset.PackageName.ToString());
+					OutTestCommands.Add(MapAssetPath + TEXT(";") + MapPackageName);
 					OutTestMapAssets.AddUnique(MapAssetPath);
 				}
 			}
diff --git a/Engine/Source/Runtime/AutomationTest/Private/AutomationTestExcludelist.cpp b/Engine/Source/Runtime/AutomationTest/Private/AutomationTestExcludelist.cpp
--- a/Engine/Source/Runtime/AutomationTest/Private/AutomationTestExcludelist.cpp
+++ b/Engine/Source/Runtime/AutomationTest/Private/AutomationTestExcludelist.cpp
@@ -57,12 +57,21 @@ FString UAutomationTestExcludelist::GetFullTestName(const FAutomationTestExclude
 		return ExcludelistEntry.FullTestName;
 	}
 
-	FString Map = ExcludelistEntry.Map.ToString().TrimStartAndEnd();
 	FString ListName = ExcludelistEntry.Test.ToString().TrimStartAndEnd();
 
-	if (Map.StartsWith(TEXT("/")) && !ListName.StartsWith(FunctionalTestsPreFix))
+	bool IsFunctionalTest = ListName.StartsWith(FunctionalTestsPreFix);
+
+	// Backcomp - merge Map and Test properties of Functional Test
+	FString Map = ExcludelistEntry.Map.ToString().TrimStartAndEnd();
+	if (Map.StartsWith(TEXT("/")) && !IsFunctionalTest)
 	{
 		ListName = FunctionalTestsPreFix + Map + TEXT(".") + ListName;
+		IsFunctionalTest = true;
+	}
+	// Backcomp - Convert package path by using dot syntax instead of /
+	if (IsFunctionalTest)
+	{
+		ListName = ListName.Replace(TEXT("./Game/"), TEXT(".")).Replace(TEXT("/"), TEXT("."));
 	}
 
 	return ListName.ToLower();
@@ -102,13 +111,8 @@ void UAutomationTestExcludelist::RemoveFromExcludeTest(const FString& TestName)
 
 bool UAutomationTestExcludelist::IsTestExcluded(const FString& TestName, const FString& RHI, FName* OutReason, bool* OutWarn)
 {
-	if (auto Entry = GetExcludeTestEntry(TestName))
+	if (auto Entry = GetExcludeTestEntry(TestName, RHI))
 	{
-		if (!Entry->ShouldExcludeForRHI(RHI))
-		{
-			return false;
-		}
-
 		if (OutReason != nullptr)
 		{
 			*OutReason = Entry->Reason;
@@ -125,7 +129,7 @@ bool UAutomationTestExcludelist::IsTestExcluded(const FString& TestName, const F
 	return false;
 }
 
-FAutomationTestExcludelistEntry* UAutomationTestExcludelist::GetExcludeTestEntry(const FString& TestName)
+FAutomationTestExcludelistEntry* UAutomationTestExcludelist::GetExcludeTestEntry(const FString& TestName, const FString& RHI)
 {
 	if (TestName.IsEmpty())
 		return nullptr;
@@ -138,7 +142,10 @@ FAutomationTestExcludelistEntry* UAutomationTestExcludelist::GetExcludeTestEntry
 		{
 			if (NameToCompare.Len() == Entry.FullTestName.Len() || NameToCompare.Mid(Entry.FullTestName.Len(), 1) == TEXT("."))
 			{
-				return &Entry;
+				if (RHI.IsEmpty() || Entry.RHIs.Num() == 0 || Entry.RHIs.Contains(*RHI))
+				{
+					return &Entry;
+				}
 			}
 		}
 	}
diff --git a/Engine/Source/Runtime/AutomationTest/Public/AutomationTestExcludelist.h b/Engine/Source/Runtime/AutomationTest/Public/AutomationTestExcludelist.h
--- a/Engine/Source/Runtime/AutomationTest/Public/AutomationTestExcludelist.h
+++ b/Engine/Source/Runtime/AutomationTest/Public/AutomationTestExcludelist.h
@@ -114,17 +114,6 @@ struct FAutomationTestExcludelistEntry
 		bIsPropagated = false;
 	}
 
-	/* Confirm that test should be excluded for that RHI */
-	bool ShouldExcludeForRHI(const FString& InRHI) const
-	{
-		if (!InRHI.IsEmpty() && RHIs.Num() != 0 && !RHIs.Contains(*InRHI))
-		{
-			return false;
-		}
-
-		return true;
-	}
-
 	/* Hold full test name/path */
 	FString FullTestName;
 	/* Is the entry comes from propagation */
@@ -167,7 +156,7 @@ public:
 	void AddToExcludeTest(const FString& TestName, const FAutomationTestExcludelistEntry& ExcludelistEntry);
 	void RemoveFromExcludeTest(const FString& TestName);
 	bool IsTestExcluded(const FString& TestName, const FString& RHI = TEXT(""), FName* OutReason = nullptr, bool* OutWarn = nullptr);
-	FAutomationTestExcludelistEntry* GetExcludeTestEntry(const FString& TestName);
+	FAutomationTestExcludelistEntry* GetExcludeTestEntry(const FString& TestName, const FString& RHI = TEXT(""));
 
 	void SaveConfig();
 	FString GetConfigFilename() { return UObject::GetDefaultConfigFilename(); } const
