Fix NoPCH build break - AutomationState.h:25:9: error: incomplete type 'UEnum' named in nested name specifier

diff --git a/Engine/Source/Runtime/AutomationTest/Public/AutomationState.h b/Engine/Source/Runtime/AutomationTest/Public/AutomationState.h
--- a/Engine/Source/Runtime/AutomationTest/Public/AutomationState.h
+++ b/Engine/Source/Runtime/AutomationTest/Public/AutomationState.h
@@ -4,6 +4,7 @@
 
 #include "CoreTypes.h"
 #include "UObject/ObjectMacros.h"
+#include "UObject/Class.h"
 
 #include "AutomationState.generated.h"
 
