Fix Exclusion list not working properly on Runtime Game

#rnx

diff --git a/Engine/Source/Runtime/AutomationTest/Public/AutomationTestExcludelist.h b/Engine/Source/Runtime/AutomationTest/Public/AutomationTestExcludelist.h
--- a/Engine/Source/Runtime/AutomationTest/Public/AutomationTestExcludelist.h
+++ b/Engine/Source/Runtime/AutomationTest/Public/AutomationTestExcludelist.h
@@ -8,20 +8,43 @@
 UENUM()
 enum class ETEST_RHI_Options
 {
-	DirectX11	UMETA(DisplayName = "DirectX 11"),
-	DirectX12	UMETA(DisplayName = "DirectX 12"),
-	Vulkan		UMETA(DisplayName = "Vulkan"),
-	Metal		UMETA(DisplayName = "Metal"),
-	Null		UMETA(DisplayName = "Null")
+	DirectX11,
+	DirectX12,
+	Vulkan,
+	Metal,
+	Null
 };
 
+inline FString LexToString(ETEST_RHI_Options Option)
+{
+	switch (Option)
+	{
+	case ETEST_RHI_Options::DirectX11:    return TEXT("DirectX 11");
+	case ETEST_RHI_Options::DirectX12:    return TEXT("DirectX 12");
+	case ETEST_RHI_Options::Vulkan:       return TEXT("Vulkan");
+	case ETEST_RHI_Options::Metal:        return TEXT("Metal");
+	case ETEST_RHI_Options::Null:         return TEXT("Null");
+	default:                              return TEXT("Unknown");
+	}
+}
+
 UENUM()
 enum class ETEST_RHI_FeatureLevel_Options
 {
-	SM5 UMETA(DisplayName = "SM5"),
-	SM6 UMETA(DisplayName = "SM6")
+	SM5,
+	SM6
 };
 
+inline FString LexToString(ETEST_RHI_FeatureLevel_Options Option)
+{
+	switch (Option)
+	{
+	case ETEST_RHI_FeatureLevel_Options::SM5:   return TEXT("SM5");
+	case ETEST_RHI_FeatureLevel_Options::SM6:   return TEXT("SM6");
+	default:                                    return TEXT("Unknown");
+	}
+}
+
 USTRUCT()
 struct FAutomationTestExcludeOptions
 {
@@ -39,7 +62,7 @@ struct FAutomationTestExcludeOptions
 				int32 Num_Flags = Enum->NumEnums() - 1;
 				for (int32 i = 0; i < Num_Flags; i++)
 				{
-					NameSet.Add(*Enum->GetDisplayNameTextByIndex(i).ToString());
+					NameSet.Add(*LexToString((EnumType)Enum->GetValueByIndex(i)));
 				}
 			}
 		}
@@ -47,12 +70,6 @@ struct FAutomationTestExcludeOptions
 		return NameSet;
 	}
 
-	template<typename EnumType>
-	static FString GetRHIOptionName(EnumType Value)
-	{
-		return StaticEnum<EnumType>()->GetDisplayNameTextByValue((int)Value).ToString();
-	}
-
 	/* Name of the target test */
 	UPROPERTY(VisibleAnywhere, Category = ExcludeTestOptions)
 	FName Test;
diff --git a/Engine/Source/Runtime/AutomationWorker/Private/AutomationWorkerModule.cpp b/Engine/Source/Runtime/AutomationWorker/Private/AutomationWorkerModule.cpp
--- a/Engine/Source/Runtime/AutomationWorker/Private/AutomationWorkerModule.cpp
+++ b/Engine/Source/Runtime/AutomationWorker/Private/AutomationWorkerModule.cpp
@@ -493,7 +493,7 @@ TSet<FName> GetRHIForAutomation()
 	FString RHI = FApp::GetGraphicsRHI();
 	if (RHI.IsEmpty())
 	{
-		RHI = FAutomationTestExcludeOptions::GetRHIOptionName(ETEST_RHI_Options::Null);
+		RHI = LexToString(ETEST_RHI_Options::Null);
 		return TSet<FName> {FName(RHI), FName(RHI)};
 	}
 
@@ -566,7 +566,6 @@ void FAutomationWorkerModule::HandleRunTestsMessage( const FAutomationWorkerRunT
 	FString SkipReason;
 	bool bWarn(false);
 	FAutomationTestFramework& AutomationTestFramework = FAutomationTestFramework::Get();
-	UAutomationTestExcludelist* Excludelist = UAutomationTestExcludelist::Get();
 	if (!AutomationTestFramework.CanRunTestInEnvironment(Message.TestName, &SkipReason, &bWarn)
 		|| IsTestExcluded(Message.FullTestPath, &SkipReason, &bWarn))
 	{
