Support load and save of platform specific test exclusion list from the Test Automation window

diff --git a/Engine/Config/BaseEngine.ini b/Engine/Config/BaseEngine.ini
--- a/Engine/Config/BaseEngine.ini
+++ b/Engine/Config/BaseEngine.ini
@@ -3152,6 +3152,20 @@ bUseEntitlementsForPrograms=false
 MacSigningIdentity=Apple Development
 IOSSigningIdentity=Apple Development
 
+[AutomationTestExcludelistSettings:Windows]
++SupportedRHIs=DirectX 11
++SupportedRHIs=DirectX 12
++SupportedRHIs=Vulkan
++SupportedRHIs=Null
+
+[AutomationTestExcludelistSettings:Linux]
++SupportedRHIs=Vulkan
++SupportedRHIs=Null
+
+[AutomationTestExcludelistSettings:Mac]
++SupportedRHIs=Metal
++SupportedRHIs=Null
+
 [HMDPluginPriority]
 ; Since SteamVR also works with the Oculus Rift, give priority to the native Oculus plugin before trying SteamVR.
 ; OpenXR is production ready and, if enabled, should be the priority plugin. For 4.27, we still need to prioritize
diff --git a/Engine/Source/Developer/AutomationController/Private/AutomationReport.cpp b/Engine/Source/Developer/AutomationController/Private/AutomationReport.cpp
--- a/Engine/Source/Developer/AutomationController/Private/AutomationReport.cpp
+++ b/Engine/Source/Developer/AutomationController/Private/AutomationReport.cpp
@@ -871,7 +871,18 @@ bool FAutomationReport::IsToBeSkipped(FName* OutReason, bool* OutWarn) const
 	{
 		if (OutReason != nullptr)
 		{
-			*OutReason = ExcludeTestInfo.Reason;
+			if (!ExcludeTestInfo.Platforms.IsEmpty())
+			{
+				*OutReason = *(SetToShortString(ExcludeTestInfo.Platforms) + TEXT(": ") + ExcludeTestInfo.Reason.ToString());
+			}
+			else if (!ExcludeTestInfo.RHIs.IsEmpty())
+			{
+				*OutReason = *(SetToShortString(ExcludeTestInfo.RHIs) + TEXT(": ") + ExcludeTestInfo.Reason.ToString());
+			}
+			else
+			{
+				*OutReason = ExcludeTestInfo.Reason;
+			}
 		}
 
 		if (OutWarn != nullptr)
@@ -959,7 +970,8 @@ void FAutomationReport::SetSkipFlag(bool bEnableSkip, const FAutomationTestExclu
 			auto Entry = ExcludedTestCached->GetExcludeTestEntry(TestInfo.GetFullTestPath());
 			if (Entry != nullptr && ExcludeTestInfo.RemoveConditions(*Entry))
 			{
-				// Branch off the template with a different exclusion condition set
+				// Branch off the template with the difference exclusion condition set
+				// Remove the overlapping conditions as they are redundant since the parent has precedence 
 				Template = &ExcludeTestInfo;
 			}
 			ExcludedTestCached->AddToExcludeTest(TestInfo.GetFullTestPath(), *Template);
@@ -993,7 +1005,7 @@ void FAutomationReport::SetSkipFlag(bool bEnableSkip, const FAutomationTestExclu
 	if (!bFromPropagation)
 	{
 		// Save config only at the end of the recursion
-		ExcludedTestCached->SaveConfig();
+		ExcludedTestCached->SaveToConfigs();
 	}
 }
 
diff --git a/Engine/Source/Developer/AutomationWindow/Private/SAutomationTestItem.cpp b/Engine/Source/Developer/AutomationWindow/Private/SAutomationTestItem.cpp
--- a/Engine/Source/Developer/AutomationWindow/Private/SAutomationTestItem.cpp
+++ b/Engine/Source/Developer/AutomationWindow/Private/SAutomationTestItem.cpp
@@ -16,6 +16,7 @@
 #include "SSimpleComboButton.h"
 #include "AutomationWindowStyle.h"
 #include "AutomationTestExcludelist.h"
+#include "AutomationTestPlatform.h"
 #include "SAutomationWindow.h"
 
 #if WITH_EDITOR
@@ -72,8 +73,9 @@ TSharedRef<SWidget> SAutomationTestItem::GenerateWidgetForColumn( const FName& C
 			.AutoWidth()
 			[
 				SNew(SButton)
-				.ButtonStyle(FAutomationWindowStyle::Get(), "NoBorder")
+				.ButtonStyle(FAutomationWindowStyle::Get(), "SimpleButton")
 				.ToolTipText(this, &SAutomationTestItem::GetExcludeReason)
+				.IsEnabled(this, &SAutomationTestItem::CanSkipFlagBeChanged)
 				.OnClicked(FOnClicked::CreateSP(this, &SAutomationTestItem::SetSkipFlag))
 #if WITH_EDITOR
 				.Cursor_Lambda([this]() {return this->IsLocalSession ? EMouseCursor::Hand : EMouseCursor::Default;})
@@ -405,6 +407,11 @@ FSlateColor SAutomationTestItem::IsToBeSkipped_GetColorAndOpacity() const
 	{
 		return FLinearColor(1.0f, 1.0f, 1.0f, 0.4f);
 	}
+	// Identify visually if the test is to be skipped on certain condition.
+	if (TestStatus->IsToBeSkippedOnConditions())
+	{
+		return FLinearColor(1.0f, 1.0f, 0.2f, 0.6f);
+	}
 	return FLinearColor(1.0f, 1.0f, 1.0f, 1.0f);
 }
 
@@ -451,47 +458,72 @@ FReply SAutomationTestItem::SetSkipFlag()
 	return FReply::Handled();
 }
 
-template<typename EnumType>
-void GenerateRHIMenuContentFromExcludeOptions(TSharedPtr<FAutomationTestExcludeOptions> Options, FMenuBuilder* MenuBuilder)
+bool SAutomationTestItem::CanSkipFlagBeChanged() const
 {
-	static const TSet<FName> AllRHI_OptionNames = FAutomationTestExcludeOptions::GetAllRHIOptionNames<EnumType>();
-	for (auto& Item : AllRHI_OptionNames)
+	return WITH_EDITOR && IsLocalSession && !TestStatus->IsToBeSkippedByPropagation();
+}
+
+#if WITH_EDITOR
+void PopulateMenuContent(FMenuBuilder* MenuBuilder, TSet<FName>* OptionsDestination, const TSet<FName>& ItemNames)
+{
+	for (const FName& Item : ItemNames)
 	{
 		TSharedRef<SWidget> FlagWidget =
 			SNew(SCheckBox)
-			.IsChecked(Options->RHIs.Contains(Item))
-			.OnCheckStateChanged_Lambda([Options, Item](ECheckBoxState NewState)
+			.IsChecked(OptionsDestination->Contains(Item))
+			.OnCheckStateChanged_Lambda([OptionsDestination, Item](ECheckBoxState NewState)
 				{
 					if (NewState == ECheckBoxState::Checked)
 					{
-						Options->RHIs.Add(Item);
+						OptionsDestination->Add(Item);
 					}
 					else
 					{
-						Options->RHIs.Remove(Item);
+						OptionsDestination->Remove(Item);
 					}
 				})
 			.Padding(FMargin(4.0f, 0.0f))
-			.Content()
-			[
-				SNew(STextBlock)
-				.Text(FText::FromName(Item))
-			];
+					.Content()
+					[
+						SNew(STextBlock)
+						.Text(FText::FromName(Item))
+					];
 
 		MenuBuilder->AddWidget(FlagWidget, FText::GetEmpty());
 	}
 }
 
+TSharedRef<SWidget> GeneratePlatformMenuContentFromExcludeOptions(TSharedPtr<FAutomationTestExcludeOptions> Options)
+{
+	static const TSet<FName>& AllPlatform_OptionNames = AutomationTestPlatform::GetAllAvailablePlatformNames();
+	FMenuBuilder MenuBuilder(false, nullptr);
+	PopulateMenuContent(&MenuBuilder, &Options->Platforms, AllPlatform_OptionNames);
+
+	return MenuBuilder.MakeWidget();
+}
+
+FText GeneratePlatformTextFromExcludeOptions(TSharedPtr<FAutomationTestExcludeOptions> Options)
+{
+	if (Options->Platforms.Num() == 0)
+	{
+		return LOCTEXT("ExcludeOptions_Platform_All", "All Platforms");
+	}
+
+	return FText::FromString(SetToString(Options->Platforms));
+}
+
 TSharedRef<SWidget> GenerateRHIMenuContentFromExcludeOptions(TSharedPtr<FAutomationTestExcludeOptions> Options)
 {
 	FMenuBuilder MenuBuilder(false, nullptr);
 
+	static const TSet<FName>& AllRHI_OptionNames = FAutomationTestExcludeOptions::GetAllRHIOptionNamesFromSettings();
 	MenuBuilder.BeginSection("AutomationWindow_ExcludeOptions_RHI", LOCTEXT("ExcludeOptions_RHI_Section", "Interfaces"));
-	GenerateRHIMenuContentFromExcludeOptions<ETEST_RHI_Options>(Options, &MenuBuilder);
+	PopulateMenuContent(&MenuBuilder, &Options->RHIs, AllRHI_OptionNames);
 	MenuBuilder.EndSection();
 
+	static const TSet<FName>& AllRHI_FeatureLevel_OptionNames = FAutomationTestExcludeOptions::GetAllRHIOptionNames<ETEST_RHI_FeatureLevel_Options>();
 	MenuBuilder.BeginSection("AutomationWindow_ExcludeOptions_RHI_FeatureLevel", LOCTEXT("ExcludeOptions_RHI_FeatureLevel_Section", "Feature Levels"));
-	GenerateRHIMenuContentFromExcludeOptions<ETEST_RHI_FeatureLevel_Options>(Options, &MenuBuilder);
+	PopulateMenuContent(&MenuBuilder, &Options->RHIs, AllRHI_FeatureLevel_OptionNames);
 	MenuBuilder.EndSection();
 
 	return MenuBuilder.MakeWidget();
@@ -504,14 +536,9 @@ FText GenerateRHITextFromExcludeOptions(TSharedPtr<FAutomationTestExcludeOptions
 		return LOCTEXT("ExcludeOptions_RHI_All", "All Interfaces");
 	}
 
-	TArray<FString> RHIs;
-	for (auto& Item : Options->RHIs)
-	{
-		RHIs.Add(Item.ToString());
-	}
-
-	return FText::FromString(FString::Join(RHIs, TEXT(", ")));
+	return FText::FromString(SetToString(Options->RHIs));
 }
+#endif
 
 FReply SAutomationTestItem::OnEditExcludeOptionsClicked()
 {
@@ -529,7 +556,7 @@ FReply SAutomationTestItem::OnEditExcludeOptionsClicked()
 			[
 				SNew(SHorizontalBox)
 				+ SHorizontalBox::Slot()
-				.MaxWidth(50)
+				.MaxWidth(55)
 				.HAlign(HAlign_Right)
 				.VAlign(VAlign_Center)
 				.Padding(5, 0)
@@ -553,7 +580,7 @@ FReply SAutomationTestItem::OnEditExcludeOptionsClicked()
 			[
 				SNew(SHorizontalBox)
 				+ SHorizontalBox::Slot()
-				.MaxWidth(50)
+				.MaxWidth(55)
 				.HAlign(HAlign_Right)
 				.VAlign(VAlign_Center)
 				.Padding(5, 0)
@@ -577,7 +604,38 @@ FReply SAutomationTestItem::OnEditExcludeOptionsClicked()
 			[
 				SNew(SHorizontalBox)
 				+ SHorizontalBox::Slot()
-				.MaxWidth(50)
+				.MaxWidth(55)
+				.HAlign(HAlign_Right)
+				.VAlign(VAlign_Center)
+				.Padding(5, 0)
+				[
+					SNew(STextBlock)
+					.Text(LOCTEXT("ExcludeOptions_Platform", "Platforms"))
+				]
+				+ SHorizontalBox::Slot()
+				.MaxWidth(200)
+				.HAlign(HAlign_Left)
+				.VAlign(VAlign_Center)
+				[
+					SNew(SComboButton)
+					.OnGetMenuContent_Lambda([Options]() { return GeneratePlatformMenuContentFromExcludeOptions(Options); })
+					.HasDownArrow(true)
+					.ButtonContent()
+					[
+						SNew(STextBlock)
+						.Text_Lambda([Options]() { return GeneratePlatformTextFromExcludeOptions(Options); })
+						.ToolTipText_Lambda([Options]() { return GeneratePlatformTextFromExcludeOptions(Options); })
+						.OverflowPolicy(ETextOverflowPolicy::Ellipsis)
+					]
+				]
+			]
+			+ SVerticalBox::Slot()
+			.AutoHeight()
+			.Padding(5.0f)
+			[
+				SNew(SHorizontalBox)
+				+ SHorizontalBox::Slot()
+				.MaxWidth(55)
 				.HAlign(HAlign_Right)
 				.VAlign(VAlign_Center)
 				.Padding(5, 0)
@@ -586,15 +644,20 @@ FReply SAutomationTestItem::OnEditExcludeOptionsClicked()
 					.Text(LOCTEXT("ExcludeOptions_RHI", "RHIs"))
 				]
 				+ SHorizontalBox::Slot()
-				.MaxWidth(120)
+				.MaxWidth(200)
 				.HAlign(HAlign_Left)
 				.VAlign(VAlign_Center)
 				[
-					SNew(SSimpleComboButton)
+					SNew(SComboButton)
 					.OnGetMenuContent_Lambda([Options]() { return GenerateRHIMenuContentFromExcludeOptions(Options); })
-					.Text_Lambda([Options]() { return GenerateRHITextFromExcludeOptions(Options); })
-					.ToolTipText_Lambda([Options]() { return GenerateRHITextFromExcludeOptions(Options); })
 					.HasDownArrow(true)
+					.ButtonContent()
+					[
+						SNew(STextBlock)
+						.Text_Lambda([Options]() { return GenerateRHITextFromExcludeOptions(Options); })
+						.ToolTipText_Lambda([Options]() { return GenerateRHITextFromExcludeOptions(Options); })
+						.OverflowPolicy(ETextOverflowPolicy::Ellipsis)
+					]
 				]
 			]
 			+ SVerticalBox::Slot()
@@ -603,7 +666,7 @@ FReply SAutomationTestItem::OnEditExcludeOptionsClicked()
 			[
 				SNew(SHorizontalBox)
 				+ SHorizontalBox::Slot()
-				.MaxWidth(50)
+				.MaxWidth(55)
 				.HAlign(HAlign_Right)
 				.VAlign(VAlign_Center)
 				.Padding(5, 0)
diff --git a/Engine/Source/Developer/AutomationWindow/Private/SAutomationTestItem.h b/Engine/Source/Developer/AutomationWindow/Private/SAutomationTestItem.h
--- a/Engine/Source/Developer/AutomationWindow/Private/SAutomationTestItem.h
+++ b/Engine/Source/Developer/AutomationWindow/Private/SAutomationTestItem.h
@@ -79,7 +79,7 @@ protected:
 	*
 	* @return true if the test is enabled.
 	*/
-	ECheckBoxState IsTestEnabled( ) const;
+	ECheckBoxState IsTestEnabled() const;
 
 	/**
 	* Is the test inside exclude list. This info is taken from Config/DefaultEngine.ini, section [AutomationTestExcludelist].
@@ -121,6 +121,13 @@ protected:
 	*/
 	FReply SetSkipFlag();
 
+	/**
+	* Can the test skip flag be changed.
+	*
+	* @return true if the test skip flag can be changed.
+	*/
+	bool CanSkipFlagBeChanged() const;
+
 	/**
 	* Open Exclude test options editor
 	*
diff --git a/Engine/Source/Runtime/AutomationTest/Private/AutomationTestExcludelist.cpp b/Engine/Source/Runtime/AutomationTest/Private/AutomationTestExcludelist.cpp
--- a/Engine/Source/Runtime/AutomationTest/Private/AutomationTestExcludelist.cpp
+++ b/Engine/Source/Runtime/AutomationTest/Private/AutomationTestExcludelist.cpp
@@ -2,6 +2,8 @@
 
 #include "AutomationTestExcludelist.h"
 
+#include "Misc/ConfigCacheIni.h"
+
 #include UE_INLINE_GENERATED_CPP_BY_NAME(AutomationTestExcludelist)
 
 #if WITH_EDITOR
@@ -17,66 +19,119 @@ DEFINE_LOG_CATEGORY_STATIC(LogAutomationTestExcludelist, Log, All);
 
 static const FString FunctionalTestsPreFix = TEXT("Project.Functional Tests.");
 
-void UAutomationTestExcludelist::OverrideConfigSection(FString& SectionName)
-{
-	SectionName = TEXT("AutomationTestExcludelist");
-}
-
-UAutomationTestExcludelist* UAutomationTestExcludelist::Get()
+void SortExcludelist(TMap<FString, FAutomationTestExcludelistEntry>& List)
 {
-	return GetMutableDefault<UAutomationTestExcludelist>();
+	// Sort in alphabetical order, shortest to longest of key property.
+	// That is to naturally gives priority to parent suite over individual test exclusion when calling GetExcludeTestEntry(TestName).
+	List.KeySort([](const FString& A, const FString& B)
+		{
+			return A < B;
+		});
 }
 
 void SortExcludelist(TArray<FAutomationTestExcludelistEntry>& List)
 {
 	// Sort in alphabetical order, shortest to longest of FullTestName property.
-	// That is to naturally gives priority to parent suite over invidual test exclusion when calling GetExcludeTestEntry(TestName).
+	// That is to naturally gives priority to parent suite over individual test exclusion when calling GetExcludeTestEntry(TestName).
 	List.Sort([](const FAutomationTestExcludelistEntry& A, const FAutomationTestExcludelistEntry& B)
+		{
+			return A.FullTestName < B.FullTestName;
+		});
+}
+
+void FAutomationTestExcludelistEntry::Finalize()
+{
+	if (!IsEmpty())
 	{
-		return A.FullTestName < B.FullTestName;
-	});
+		return;
+	}
+
+	FString TestStr = Test.ToString().TrimStartAndEnd();
+
+	bool IsFunctionalTest = TestStr.StartsWith(FunctionalTestsPreFix);
+
+	// Backward compatibility - merge Map and Test properties of Functional Test
+	FString MapStr = Map.ToString().TrimStartAndEnd();
+	if (MapStr.StartsWith(TEXT("/")) && !IsFunctionalTest)
+	{
+		TestStr = FunctionalTestsPreFix + MapStr + TEXT(".") + TestStr;
+		IsFunctionalTest = true;
+	}
+	// Backward compatibility - Convert package path by using dot syntax instead of /
+	if (IsFunctionalTest)
+	{
+		TestStr = TestStr.Replace(TEXT("./Game/"), TEXT(".")).Replace(TEXT("/"), TEXT("."));
+	}
+
+	FullTestName = TestStr.ToLower();
 }
 
-void UAutomationTestExcludelist::PostInitProperties()
+FString FAutomationTestExcludelistEntry::GetStringForHash() const
 {
-	Super::PostInitProperties();
+	return FullTestName + Reason.ToString() + SetToString(RHIs) + (Warn? TEXT("1") : TEXT("0"));
+}
 
-	for (auto& Entry : ExcludeTest)
+UAutomationTestExcludelist* UAutomationTestExcludelist::Get()
+{
+	UAutomationTestExcludelist* Obj = GetMutableDefault<UAutomationTestExcludelist>();
+	if (!Obj->DefaultConfig)
 	{
-		if (Entry.IsEmpty())
-		{
-			Entry.FullTestName = GetFullTestName(Entry);
-		}
+		Obj->Initialize();
 	}
 
-	SortExcludelist(ExcludeTest);
+	return Obj;
 }
 
-FString UAutomationTestExcludelist::GetFullTestName(const FAutomationTestExcludelistEntry& ExcludelistEntry)
+void UAutomationTestExcludelist::Initialize()
 {
-	if (!ExcludelistEntry.IsEmpty())
+	DefaultConfig = GetMutableDefault<UAutomationTestExcludelistConfig>();
+	if (PlatformConfigs.IsEmpty())
 	{
-		return ExcludelistEntry.FullTestName;
+		LoadPlatformConfigs();
+		PopulateEntries();
 	}
+}
 
-	FString ListName = ExcludelistEntry.Test.ToString().TrimStartAndEnd();
+void UAutomationTestExcludelist::LoadPlatformConfigs()
+{
+	PlatformConfigs.Empty();
+
+	for (auto* PlatformSettings : AutomationTestPlatform::GetAllPlatformsSettings(UAutomationTestExcludelistConfig::StaticClass()))
+	{
+		PlatformConfigs.Emplace(PlatformSettings->GetPlatformName(), CastChecked<UAutomationTestExcludelistConfig>(PlatformSettings));
+	}
+}
 
-	bool IsFunctionalTest = ListName.StartsWith(FunctionalTestsPreFix);
+void UAutomationTestExcludelist::PopulateEntries()
+{
+	LLM_SCOPE_BYNAME(TEXT("AutomationTest/Settings"));
+	Entries.Empty();
 
-	// Backcomp - merge Map and Test properties of Functional Test
-	FString Map = ExcludelistEntry.Map.ToString().TrimStartAndEnd();
-	if (Map.StartsWith(TEXT("/")) && !IsFunctionalTest)
+	// Populate with default first
+	for (const FAutomationTestExcludelistEntry& Entry : DefaultConfig->GetEntries())
 	{
-		ListName = FunctionalTestsPreFix + Map + TEXT(".") + ListName;
-		IsFunctionalTest = true;
+		Entries.Emplace(Entry.FullTestName, Entry);
 	}
-	// Backcomp - Convert package path by using dot syntax instead of /
-	if (IsFunctionalTest)
+	// Merge platforms with default config
+	for (auto& Config : PlatformConfigs)
 	{
-		ListName = ListName.Replace(TEXT("./Game/"), TEXT(".")).Replace(TEXT("/"), TEXT("."));
+		for (const FAutomationTestExcludelistEntry& PlatformEntry : Config.Value->GetEntries())
+		{
+			if (FAutomationTestExcludelistEntry* Entry = Entries.Find(PlatformEntry.FullTestName))
+			{
+				Entry->Platforms.Add(Config.Key);
+				Entry->RHIs.Append(PlatformEntry.RHIs);
+			}
+			else
+			{
+				Entries.Emplace(
+					PlatformEntry.FullTestName,
+					PlatformEntry
+				).Platforms.Add(Config.Key);
+			}
+		}
 	}
-
-	return ListName.ToLower();
+	SortExcludelist(Entries);
 }
 
 void UAutomationTestExcludelist::AddToExcludeTest(const FString& TestName, const FAutomationTestExcludelistEntry& ExcludelistEntry)
@@ -87,11 +142,9 @@ void UAutomationTestExcludelist::AddToExcludeTest(const FString& TestName, const
 	{
 		NewEntry.Map = TEXT("");
 	}
-
-	NewEntry.FullTestName = GetFullTestName(NewEntry);
-
-	ExcludeTest.Add(NewEntry);
-	SortExcludelist(ExcludeTest);
+	NewEntry.Finalize();
+	Entries.Emplace(NewEntry.FullTestName, NewEntry);
+	SortExcludelist(Entries);
 }
 
 void UAutomationTestExcludelist::RemoveFromExcludeTest(const FString& TestName)
@@ -99,21 +152,25 @@ void UAutomationTestExcludelist::RemoveFromExcludeTest(const FString& TestName)
 	if (TestName.IsEmpty())
 		return;
 
-	const FString NameToCompare = TestName.TrimStartAndEnd().ToLower();
+	Entries.Remove(TestName.TrimStartAndEnd().ToLower());
+	SortExcludelist(Entries);
+}
 
-	for (int i = 0; i < ExcludeTest.Num(); ++i)
-	{
-		if (ExcludeTest[i].FullTestName == NameToCompare)
-		{
-			ExcludeTest.RemoveAt(i);
-			return;
-		}
-	}
+bool UAutomationTestExcludelist::IsTestExcluded(const FString& TestName) const
+{
+	static const FName None;
+	static const TSet<FName> EmptySet;
+	return IsTestExcluded(TestName, None, EmptySet, nullptr, nullptr);
+}
+
+bool UAutomationTestExcludelist::IsTestExcluded(const FString& TestName, const TSet<FName>& RHI, FName* OutReason, bool* OutWarn) const
+{
+	return IsTestExcluded(TestName, FPlatformProperties::IniPlatformName(), RHI, OutReason, OutWarn);
 }
 
-bool UAutomationTestExcludelist::IsTestExcluded(const FString& TestName, const TSet<FName>& RHI, FName* OutReason, bool* OutWarn)
+bool UAutomationTestExcludelist::IsTestExcluded(const FString & TestName, const FName& Platform, const TSet<FName>& RHI, FName * OutReason, bool* OutWarn) const
 {
-	if (auto Entry = GetExcludeTestEntry(TestName, RHI))
+	if (const auto Entry = GetExcludeTestEntry(TestName, Platform, RHI))
 	{
 		if (OutReason != nullptr)
 		{
@@ -131,34 +188,132 @@ bool UAutomationTestExcludelist::IsTestExcluded(const FString& TestName, const T
 	return false;
 }
 
-FAutomationTestExcludelistEntry* UAutomationTestExcludelist::GetExcludeTestEntry(const FString& TestName, const TSet<FName>& RHI)
+FString UAutomationTestExcludelist::GetConfigFilename() const
+{
+	return DefaultConfig->GetConfigFilename();
+}
+
+FString UAutomationTestExcludelist::GetConfigFilenameForEntry(const FAutomationTestExcludelistEntry& Entry) const
+{
+	return GetConfigFilenameForEntry(Entry, FPlatformProperties::IniPlatformName());
+}
+
+FString UAutomationTestExcludelist::GetConfigFilenameForEntry(const FAutomationTestExcludelistEntry& Entry, const FName& PlatformName) const
+{
+	if (Entry.Platforms.IsEmpty())
+	{
+		return DefaultConfig->GetConfigFilename();
+	}
+
+	// Align with current platform
+	if (const TObjectPtr<UAutomationTestExcludelistConfig>* ConfigPtr = PlatformConfigs.Find(PlatformName))
+	{
+		return (*ConfigPtr)->GetConfigFilename();
+	}
+
+	// Otherwise take the first item from the entry platform list
+	FName FirstItem = Entry.Platforms[Entry.Platforms.begin().GetId()];
+	if (const TObjectPtr<UAutomationTestExcludelistConfig>* ConfigPtr = PlatformConfigs.Find(FirstItem))
+	{
+		return (*ConfigPtr)->GetConfigFilename();
+	}
+
+	return TEXT("");
+}
+
+void UAutomationTestExcludelist::SaveToConfigs()
+{
+	// Reset the cached configs
+	DefaultConfig->Reset();
+	for (auto& Config : PlatformConfigs)
+	{
+		Config.Value->Reset();
+	}
+
+	// Populate configs
+	for (auto& EntryPair : Entries)
+	{
+		if (EntryPair.Value.Platforms.IsEmpty())
+		{
+			DefaultConfig->AddEntry(EntryPair.Value);
+		}
+		else
+		{
+			FAutomationTestExcludelistEntry PlatformEntry = EntryPair.Value;
+			for (const FName& PlatformName : EntryPair.Value.Platforms)
+			{
+				if (!EntryPair.Value.RHIs.IsEmpty())
+				{
+					// Filter in only the RHIs that are relevant for the platform
+					PlatformEntry.RHIs = EntryPair.Value.RHIs.Intersect(FAutomationTestExcludeOptions::GetPlatformRHIOptionNamesFromSettings(PlatformName));
+				}
+				if (TObjectPtr<UAutomationTestExcludelistConfig>* ConfigPtr = PlatformConfigs.Find(PlatformName))
+				{
+					(*ConfigPtr)->AddEntry(PlatformEntry);
+				}
+				else
+				{
+					UAutomationTestExcludelistConfig* Config = CastChecked<UAutomationTestExcludelistConfig>(UAutomationTestPlatformSettings::Create(UAutomationTestExcludelistConfig::StaticClass(), PlatformName.ToString()));
+					Config->AddEntry(PlatformEntry);
+					PlatformConfigs.Emplace(PlatformName, Config);
+				}
+			}
+		}
+	}
+
+	// Save the configs
+	DefaultConfig->SaveConfig();
+	for (auto& Config : PlatformConfigs)
+	{
+		Config.Value->SaveConfig();
+	}
+}
+
+const FAutomationTestExcludelistEntry* UAutomationTestExcludelist::GetExcludeTestEntry(const FString& TestName) const
+{
+	static const FName None;
+	static const TSet<FName> EmptySet;
+	return GetExcludeTestEntry(TestName, None, EmptySet);
+}
+
+const FAutomationTestExcludelistEntry* UAutomationTestExcludelist::GetExcludeTestEntry(const FString& TestName, const TSet<FName>& RHI) const
+{
+	return GetExcludeTestEntry(TestName, FPlatformProperties::IniPlatformName(), RHI);
+}
+
+const FAutomationTestExcludelistEntry* UAutomationTestExcludelist::GetExcludeTestEntry(const FString& TestName, const FName& Platform, const TSet<FName>& RHI) const
 {
 	if (TestName.IsEmpty())
 		return nullptr;
 
 	const FString NameToCompare = TestName.TrimStartAndEnd().ToLower();
 
-	FAutomationTestExcludelistEntry* OutEntry = nullptr;
+	const FAutomationTestExcludelistEntry* OutEntry = nullptr;
 
-	for (auto& Entry : ExcludeTest)
+	for (auto& EntryPair : Entries)
 	{
-		if(NameToCompare.StartsWith(Entry.FullTestName))
+		if (NameToCompare.StartsWith(EntryPair.Key))
 		{
-			if (NameToCompare.Len() == Entry.FullTestName.Len() || NameToCompare.Mid(Entry.FullTestName.Len(), 1) == TEXT("."))
+			if (NameToCompare.Len() == EntryPair.Key.Len() || NameToCompare.Mid(EntryPair.Key.Len(), 1) == TEXT("."))
 			{
-				if (Entry.RHIs.Num() == 0)
+				if (!Platform.IsNone() && !EntryPair.Value.Platforms.IsEmpty() && !EntryPair.Value.Platforms.Contains(Platform))
 				{
-					return &Entry;
+					continue;
+				}
+
+				if (EntryPair.Value.RHIs.IsEmpty())
+				{
+					return &EntryPair.Value;
 				}
 				if (RHI.IsEmpty())
 				{
-					OutEntry = &Entry;
+					OutEntry = &EntryPair.Value;
 					continue;
 				}
-				const int8 IntersectNum = RHI.Intersect(Entry.RHIs).Num();
-				if (IntersectNum > 0 && IntersectNum == Entry.NumRHIType())
+				const int8 IntersectNum = RHI.Intersect(EntryPair.Value.RHIs).Num();
+				if (IntersectNum > 0 && IntersectNum == EntryPair.Value.NumRHIType())
 				{
-					return &Entry;
+					return &EntryPair.Value;
 				}
 			}
 		}
@@ -167,6 +322,56 @@ FAutomationTestExcludelistEntry* UAutomationTestExcludelist::GetExcludeTestEntry
 	return OutEntry;
 }
 
+void UAutomationTestExcludelistConfig::InitializeSettingsDefault()
+{
+	Reset();
+}
+
+void UAutomationTestExcludelistConfig::Reset()
+{
+	ExcludeTest.Empty();
+	EntriesHash = FSHAHash();
+}
+
+void UAutomationTestExcludelistConfig::AddEntry(const FAutomationTestExcludelistEntry& Entry)
+{
+	ExcludeTest.Add(Entry);
+	UpdateHash(Entry);
+}
+
+void UAutomationTestExcludelistConfig::UpdateHash(const FAutomationTestExcludelistEntry& Entry)
+{
+	check(!Entry.IsEmpty());
+	FSHA1 SHA;
+	SHA.Update((const uint8*)&EntriesHash, sizeof(EntriesHash));
+	FString EntryString = Entry.GetStringForHash();
+	SHA.UpdateWithString(*EntryString, EntryString.Len());
+	EntriesHash = SHA.Finalize();
+}
+
+const TArray<FAutomationTestExcludelistEntry>& UAutomationTestExcludelistConfig::GetEntries() const
+{
+	return ExcludeTest;
+}
+
+void UAutomationTestExcludelistConfig::PostInitProperties()
+{
+	Super::PostInitProperties();
+
+	for (auto& Entry : ExcludeTest)
+	{
+		Entry.Finalize();
+	}
+	SortExcludelist(ExcludeTest);
+	// Hashing is order sensitive and depends on the Entry being finalized.
+	for (auto& Entry : ExcludeTest)
+	{
+		UpdateHash(Entry);
+	}
+	// Store the initial hash to detect dirty state.
+	SavedEntriesHash = EntriesHash;
+}
+
 #if WITH_EDITOR
 bool CheckOutOrAddFile(const FString& InFileToCheckOut)
 {
@@ -203,6 +408,26 @@ bool CheckOutOrAddFile(const FString& InFileToCheckOut)
 						bSuccessfullyCheckedOutOrAddedFile = true;
 					}
 				}
+				else if (SourceControlState->CanAdd() || SourceControlState->IsUnknown())
+				{
+					ECommandResult::Type CommandResult = SourceControlProvider.Execute(ISourceControlOperation::Create<FMarkForAdd>(), FilesToBeCheckedOut);
+					if (CommandResult == ECommandResult::Failed)
+					{
+						UE_LOG(LogAutomationTestExcludelist, Error, TEXT("Failed to mark for add the configuration file."));
+					}
+					else if (CommandResult == ECommandResult::Cancelled)
+					{
+						UE_LOG(LogAutomationTestExcludelist, Warning, TEXT("Mark for add was cancelled."));
+					}
+					else
+					{
+						bSuccessfullyCheckedOutOrAddedFile = true;
+					}
+				}
+				else if (SourceControlState->IsAdded())
+				{
+					bSuccessfullyCheckedOutOrAddedFile = true;
+				}
 			}
 			else if (!SourceControlState->IsUnknown())
 			{
@@ -242,12 +467,19 @@ bool MakeWritable(const FString& InFileToMakeWritable)
 }
 #endif
 
-void UAutomationTestExcludelist::SaveConfig()
+void UAutomationTestExcludelistConfig::SaveConfig()
 {
+	// Exit early if entries has not changed.
+	if (SavedEntriesHash == EntriesHash)
+	{
+		return;
+	}
+
 #if WITH_EDITOR
 	FString ConfigFilename = GetConfigFilename();
-	bool bIsWritable = FPlatformFileManager::Get().GetPlatformFile().FileExists(*ConfigFilename) && !FPlatformFileManager::Get().GetPlatformFile().IsReadOnly(*ConfigFilename);
-	if (!bIsWritable)
+	bool bIsFileExists = FPlatformFileManager::Get().GetPlatformFile().FileExists(*ConfigFilename);
+	bool bIsWritable = bIsFileExists && !FPlatformFileManager::Get().GetPlatformFile().IsReadOnly(*ConfigFilename);
+	if (bIsFileExists && !bIsWritable)
 	{
 		bIsWritable = CheckOutOrAddFile(ConfigFilename);
 		if (!bIsWritable)
@@ -257,14 +489,23 @@ void UAutomationTestExcludelist::SaveConfig()
 		}
 	}
 
-	if (!bIsWritable)
+	if (bIsFileExists && !bIsWritable)
 	{
 		UE_LOG(LogAutomationTestExcludelist, Error, TEXT("Failed to make the configuration file '%s' writable."), *ConfigFilename);
 	}
 	else
 #endif
 	{
-		UObject::TryUpdateDefaultConfigFile();
+		if (UObject::TryUpdateDefaultConfigFile())
+		{
+			SavedEntriesHash = EntriesHash;
+#if WITH_EDITOR
+			if (!bIsFileExists)
+			{
+				CheckOutOrAddFile(ConfigFilename);
+			}
+#endif
+		}
 	}
 }
 
diff --git a/Engine/Source/Runtime/AutomationTest/Private/AutomationTestPlatform.cpp b/Engine/Source/Runtime/AutomationTest/Private/AutomationTestPlatform.cpp
new file mode 100644
--- /dev/null
+++ b/Engine/Source/Runtime/AutomationTest/Private/AutomationTestPlatform.cpp
@@ -0,0 +1,113 @@
+// Copyright Epic Games, Inc. All Rights Reserved.
+
+#include "AutomationTestPlatform.h"
+
+#include "UObject/Package.h"
+
+#if WITH_EDITOR
+#include "Misc/CoreMisc.h"
+#include "Misc/DataDrivenPlatformInfoRegistry.h"
+#endif
+
+#include UE_INLINE_GENERATED_CPP_BY_NAME(AutomationTestPlatform)
+
+namespace AutomationTestPlatform
+{
+	const TSet<FName>& GetAllAvailablePlatformNames()
+	{
+		LLM_SCOPE_BYNAME(TEXT("AutomationTest/Settings"));
+		static TSet<FName> NameSet;
+		if (NameSet.IsEmpty())
+		{
+#if WITH_EDITOR
+			FName TargetPlatformName;
+			const TMap<FName, FDataDrivenPlatformInfo>& PlatformInfos{ FDataDrivenPlatformInfoRegistry::GetAllPlatformInfos() };
+			for (const TPair<FName, FDataDrivenPlatformInfo>& PlatformInfosMapItem : PlatformInfos)
+			{
+				TargetPlatformName = PlatformInfosMapItem.Key;
+				if (TargetPlatformName.IsNone() || FDataDrivenPlatformInfoRegistry::IsPlatformHiddenFromUI(TargetPlatformName))
+				{
+					continue;
+				}
+				const FDataDrivenPlatformInfo& PlatformInfo = PlatformInfosMapItem.Value;
+				if (PlatformInfo.bIsFakePlatform || (!PlatformInfo.bEnabledForUse))
+				{
+					continue;
+				}
+				NameSet.Add(TargetPlatformName);
+			}
+			NameSet.Sort([](const FName& A, const FName& B) { return A.ToString() < B.ToString(); });
+#else
+			NameSet.Add(FPlatformProperties::IniPlatformName());
+#endif
+		}
+
+		return NameSet;
+	}
+
+	TArray<UAutomationTestPlatformSettings*>& GetAllPlatformsSettings(TSubclassOf<UAutomationTestPlatformSettings> SettingsClass)
+	{
+		LLM_SCOPE_BYNAME(TEXT("AutomationTest/Settings"));
+		static TMap<FString, TArray<UAutomationTestPlatformSettings*>> Settings;
+		const FString SettingsName = SettingsClass->GetName();
+
+		if (!Settings.Contains(SettingsName))
+		{
+			const TSet<FName> Platforms = GetAllAvailablePlatformNames();
+			TArray<UAutomationTestPlatformSettings*> PlatformCollection;
+
+			for (const FName& PlatformName : Platforms)
+			{
+				UAutomationTestPlatformSettings* Config = UAutomationTestPlatformSettings::Create(SettingsClass, PlatformName.ToString());
+				check(Config != nullptr);
+				Config->LoadConfig();
+				Config->AddToRoot();
+				PlatformCollection.Add(Config);
+			}
+
+			Settings.Emplace(SettingsName, PlatformCollection);
+		}
+
+		return *Settings.Find(SettingsName);
+	}
+}
+
+UAutomationTestPlatformSettings* UAutomationTestPlatformSettings::Create(TSubclassOf<UAutomationTestPlatformSettings> SettingsClass, const FString InPlatformName)
+{
+	const FString PlatformConfigsName = SettingsClass->GetName() + TEXT("_") + InPlatformName;
+	UAutomationTestPlatformSettings* Config = NewObject<UAutomationTestPlatformSettings>(GetTransientPackage(), SettingsClass, FName(*PlatformConfigsName));
+	Config->PlatformName = InPlatformName;
+	Config->InitializeSettingsDefault();
+	return Config;
+}
+
+void UAutomationTestPlatformSettings::OverrideConfigSection(FString& SectionName)
+{
+	SectionName = GetSectionName();
+	if (!PlatformName.IsEmpty())
+	{
+		SectionName += TEXT(":");
+		SectionName += PlatformName;
+	}
+}
+
+const TCHAR* UAutomationTestPlatformSettings::GetConfigOverridePlatform() const
+{
+	return PlatformName.IsEmpty() ? nullptr : *PlatformName;
+}
+
+FString UAutomationTestPlatformSettings::GetConfigFilename() const
+{
+	return UObject::GetDefaultConfigFilename();
+}
+
+const FString& UAutomationTestPlatformSettings::GetPlatformName() const
+{
+	return PlatformName;
+}
+
+void UAutomationTestPlatformSettings::LoadConfig()
+{
+	UObject::LoadConfig(GetClass());
+	PostInitProperties();
+}
diff --git a/Engine/Source/Runtime/AutomationTest/Public/AutomationTestExcludelist.h b/Engine/Source/Runtime/AutomationTest/Public/AutomationTestExcludelist.h
--- a/Engine/Source/Runtime/AutomationTest/Public/AutomationTestExcludelist.h
+++ b/Engine/Source/Runtime/AutomationTest/Public/AutomationTestExcludelist.h
@@ -2,7 +2,10 @@
 
 #pragma once
 
+#include "AutomationTestPlatform.h"
 #include "CoreMinimal.h"
+#include "Misc/SecureHash.h"
+
 #include "AutomationTestExcludelist.generated.h"
 
 UENUM()
@@ -45,14 +48,69 @@ inline FString LexToString(ETEST_RHI_FeatureLevel_Options Option)
 	}
 }
 
+inline FString SetToString(const TSet<FName>& Set)
+{
+	if (Set.IsEmpty())
+	{
+		return TEXT("");
+	}
+
+	TArray<FString> List;
+	for (auto& Item : Set)
+	{
+		List.Add(Item.ToString());
+	}
+	List.Sort();
+
+	return FString::Join(List, TEXT(", "));
+}
+
+inline FString SetToShortString(const TSet<FName>& Set)
+{
+	if (Set.IsEmpty())
+	{
+		return TEXT("");
+	}
+
+	TArray<FString> List;
+	for (auto& Item : Set)
+	{
+		List.Add(Item.ToString());
+	}
+	if (List.Num() == 1)
+	{
+		return List[0];
+	}		
+	List.Sort();
+
+	return List[0] + TEXT("+");
+}
+
+UCLASS(config = Engine, defaultconfig, MinimalAPI)
+class UAutomationTestExcludelistSettings : public UAutomationTestPlatformSettings
+{
+	GENERATED_BODY()
+
+public:
+
+	UPROPERTY(Config)
+	TArray<FName> SupportedRHIs;
+
+protected:
+	virtual void InitializeSettingsDefault() { }
+	virtual FString GetSectionName() { return TEXT("AutomationTestExcludelistSettings"); }
+
+};
+
 USTRUCT()
 struct FAutomationTestExcludeOptions
 {
 	GENERATED_BODY()
 
 	template<typename EnumType>
-	static TSet<FName> GetAllRHIOptionNames()
+	static const TSet<FName>& GetAllRHIOptionNames()
 	{
+		LLM_SCOPE_BYNAME(TEXT("AutomationTest/Settings"));
 		static TSet<FName> NameSet;
 		if (NameSet.IsEmpty())
 		{
@@ -70,6 +128,38 @@ struct FAutomationTestExcludeOptions
 		return NameSet;
 	}
 
+	static const TSet<FName>& GetAllRHIOptionNamesFromSettings()
+	{
+		LLM_SCOPE_BYNAME(TEXT("AutomationTest/Settings"));
+		static TSet<FName> NameSet;
+		if (NameSet.IsEmpty())
+		{
+			for (auto Settings : AutomationTestPlatform::GetAllPlatformsSettings(UAutomationTestExcludelistSettings::StaticClass()))
+			{
+				NameSet.Append(CastChecked<UAutomationTestExcludelistSettings>(Settings)->SupportedRHIs);
+			}
+			NameSet.Sort([](const FName& A, const FName& B) { return A.ToString() < B.ToString(); });
+		}
+
+		return NameSet;
+	}
+
+	static const TSet<FName>& GetPlatformRHIOptionNamesFromSettings(const FName& Platform)
+	{
+		LLM_SCOPE_BYNAME(TEXT("AutomationTest/Settings"));
+		static TMap<FName, TSet<FName>> PlatformSettings;
+		if (PlatformSettings.IsEmpty())
+		{
+			for (auto Settings : AutomationTestPlatform::GetAllPlatformsSettings(UAutomationTestExcludelistSettings::StaticClass()))
+			{
+				PlatformSettings.Emplace(Settings->GetPlatformName()).Append(CastChecked<UAutomationTestExcludelistSettings>(Settings)->SupportedRHIs);
+			}
+		}
+
+		return PlatformSettings.FindOrAdd(Platform);
+	}
+
+
 	/* Name of the target test */
 	UPROPERTY(VisibleAnywhere, Category = ExcludeTestOptions)
 	FName Test;
@@ -82,6 +172,10 @@ struct FAutomationTestExcludeOptions
 	UPROPERTY(EditAnywhere, Category = ExcludeTestOptions)
 	TSet<FName> RHIs;
 
+	/* Options to target specific platform. No option means it should be applied to all platforms */
+	UPROPERTY(EditAnywhere, Category = ExcludeTestOptions)
+	TSet<FName> Platforms;
+
 	/* Should the Reason be reported as a warning in the log */
 	UPROPERTY(EditAnywhere, Category = ExcludeTestOptions)
 	bool Warn = false;
@@ -95,7 +189,8 @@ struct FAutomationTestExcludelistEntry
 	FAutomationTestExcludelistEntry() { }
 
 	FAutomationTestExcludelistEntry(const FAutomationTestExcludeOptions& Options)
-		: Test(Options.Test)
+		: Platforms(Options.Platforms)
+		, Test(Options.Test)
 		, Reason(Options.Reason)
 		, RHIs(Options.RHIs)
 		, Warn(Options.Warn)
@@ -108,6 +203,7 @@ struct FAutomationTestExcludelistEntry
 		Options->Reason = Reason;
 		Options->Warn = Warn;
 		Options->RHIs = RHIs;
+		Options->Platforms = Platforms;
 		
 		return Options;
 	}
@@ -117,7 +213,7 @@ struct FAutomationTestExcludelistEntry
 	{
 		int8 Num = 0;
 		// Test mentions of each RHI option types
-		static const TSet<FName> AllRHI_OptionNames = FAutomationTestExcludeOptions::GetAllRHIOptionNames<ETEST_RHI_Options>();
+		static const TSet<FName> AllRHI_OptionNames = FAutomationTestExcludeOptions::GetAllRHIOptionNamesFromSettings();
 		if (RHIs.Difference(AllRHI_OptionNames).Num() < RHIs.Num())
 		{
 			Num++;
@@ -150,13 +246,19 @@ struct FAutomationTestExcludelistEntry
 		bIsPropagated = false;
 	}
 
+	/* Make the entry specific */
+	void Finalize();
+
+	/* Output string used to generate hash */
+	FString GetStringForHash() const;
+
 	/* Has conditional exclusion */
 	bool HasConditions() const
 	{
-		return !RHIs.IsEmpty();
+		return !RHIs.IsEmpty() || !Platforms.IsEmpty();
 	}
 
-	/* Remove exclusion conditions, return true if a condition was removed */
+	/* Remove overlapping exclusion conditions, return true if at least one condition was removed */
 	bool RemoveConditions(const FAutomationTestExcludelistEntry& Entry)
 	{
 		if (!Entry.HasConditions())
@@ -165,13 +267,21 @@ struct FAutomationTestExcludelistEntry
 		}
 
 		bool GotRemoved = false;
+		int Length;
 		// Check RHIs
-		int Length = RHIs.Num();
+		Length = RHIs.Num();
 		if (Length > 0)
 		{
 			RHIs = RHIs.Difference(Entry.RHIs);
 			GotRemoved = GotRemoved || Length != RHIs.Num();
 		}
+		// Check Platforms
+		Length = Platforms.Num();
+		if (Length > 0)
+		{
+			Platforms = Platforms.Difference(Entry.Platforms);
+			GotRemoved = GotRemoved || Length != Platforms.Num();
+		}
 
 		return GotRemoved;
 	}
@@ -180,60 +290,187 @@ struct FAutomationTestExcludelistEntry
 	FString FullTestName;
 	/* Is the entry comes from propagation */
 	bool bIsPropagated = false;
+	/* Platforms to which the entry applies */
+	TSet<FName> Platforms;
 
 	// Use FName instead of FString to read params from config that aren't wrapped with quotes
 
 	/* Hold the name of the target functional test map */
-	UPROPERTY(EditDefaultsOnly, Category = AutomationTestExcludelist)
+	UPROPERTY(Config)
 	FName Map;
 
 	/* Hold the name of the target test - full test name is require here unless for functional tests */
-	UPROPERTY(EditDefaultsOnly, Category = AutomationTestExcludelist)
+	UPROPERTY(Config)
 	FName Test;
 
 	/* Reason to why the test is excluded */
-	UPROPERTY(EditDefaultsOnly, Category = AutomationTestExcludelist)
+	UPROPERTY(Config)
 	FName Reason;
 
 	/* Option to target specific RHI. Empty array means it should be applied to all RHI */
-	UPROPERTY(EditDefaultsOnly, Category = AutomationTestExcludelist)
+	UPROPERTY(Config)
 	TSet<FName> RHIs;
 
 	/* Should the Reason be reported as a warning in the log */
-	UPROPERTY(EditDefaultsOnly, Category = AutomationTestExcludelist)
+	UPROPERTY(Config)
 	bool Warn = false;
 };
 
-
 UCLASS(config = Engine, defaultconfig, MinimalAPI)
+class UAutomationTestExcludelistConfig : public UAutomationTestPlatformSettings
+{
+	GENERATED_BODY()
+
+public:
+	void Reset();
+	void AddEntry(const FAutomationTestExcludelistEntry& Entry);
+	const TArray<FAutomationTestExcludelistEntry>& GetEntries() const;
+
+	void SaveConfig();
+
+protected:
+	virtual void PostInitProperties() override;
+	virtual void InitializeSettingsDefault() override;
+	virtual FString GetSectionName() override { return TEXT("AutomationTestExcludelist"); }
+
+private:
+	void UpdateHash(const FAutomationTestExcludelistEntry& Entry);
+
+	UPROPERTY(Config)
+	TArray<FAutomationTestExcludelistEntry> ExcludeTest;
+
+	FString PlatformName;
+	/** Keep track of any changes */
+	FSHAHash EntriesHash;
+	/** Hash of what was last saved */
+	FSHAHash SavedEntriesHash;
+};
+
+UCLASS(MinimalAPI)
 class UAutomationTestExcludelist : public UObject
 {
 	GENERATED_BODY()
 
 public:
-	UAutomationTestExcludelist() {}
+	UAutomationTestExcludelist() { };
 
 	static AUTOMATIONTEST_API UAutomationTestExcludelist* Get();
 
+	/**
+	 * Add an Entry to the exclusion list for TestName
+	 *
+	 * @param TestName The full test name to be added.
+	 * @param Entry The FAutomationTestExcludelistEntry to be excluded.
+	*/
 	AUTOMATIONTEST_API void AddToExcludeTest(const FString& TestName, const FAutomationTestExcludelistEntry& ExcludelistEntry);
+
+	/**
+	 * Remove the Entry from the exclusion list for TestName
+	 *
+	 * @param TestName The full test name to be removed.
+	*/
 	AUTOMATIONTEST_API void RemoveFromExcludeTest(const FString& TestName);
-	AUTOMATIONTEST_API bool IsTestExcluded(const FString& TestName, const TSet<FName>& = TSet<FName>(), FName* OutReason = nullptr, bool* OutWarn = nullptr);
-	AUTOMATIONTEST_API FAutomationTestExcludelistEntry* GetExcludeTestEntry(const FString& TestName, const TSet<FName>& = TSet<FName>());
 
-	AUTOMATIONTEST_API void SaveConfig();
-	FString GetConfigFilename() { return UObject::GetDefaultConfigFilename(); } const
-	// It is called automatically when CDO is created, usually you don't need to call LoadConfig manually
-	void LoadConfig() { UObject::LoadConfig(GetClass()); }
+	/**
+	 * Is the TestName excluded?
+	 *
+	 * @param TestName The full test name to be tested.
+	 * @return True if the test is excluded
+	*/
+	AUTOMATIONTEST_API bool IsTestExcluded(const FString& TestName) const;
+
+	/**
+	 * Is the TestName excluded?
+	 *
+	 * @param TestName The full test name to be tested.
+	 * @param RHI The RHI the test is being evaluated against.
+	 * @param OutReason The reason why it is excluded.
+	 * @param OutWarn Whether a warning is going to be returned if the test is to be excluded.
+	 * @return True if the test is excluded
+	*/
+	AUTOMATIONTEST_API bool IsTestExcluded(const FString& TestName, const TSet<FName>& RHI, FName* OutReason, bool* OutWarn) const;
+
+	/**
+	 * Is the TestName excluded?
+	 *
+	 * @param TestName The full test name to be tested.
+	 * @param Platform The Platform the test is being evaluated against.
+	 * @param RHI The RHI the test is being evaluated against.
+	 * @param OutReason The reason why it is excluded.
+	 * @param OutWarn Whether a warning is going to be returned if the test is to be excluded.
+	 * @return True if the test is excluded
+	*/
+	AUTOMATIONTEST_API bool IsTestExcluded(const FString& TestName, const FName& Platform, const TSet<FName>& RHI, FName* OutReason, bool* OutWarn) const;
+
+	/**
+	 * Get the Entry object of the TestName if excluded
+	 *
+	 * @param TestName The full test name to be tested.
+	 * @return Return the corresponding FAutomationTestExcludelistEntry if excluded otherwise nullptr
+	*/
+	AUTOMATIONTEST_API const FAutomationTestExcludelistEntry* GetExcludeTestEntry(const FString& TestName) const;
+
+	/**
+	 * Get the Entry object of the TestName if excluded
+	 *
+	 * @param TestName The full test name to be tested.
+	 * @param Platform The Platform the test is being evaluated against.
+	 * @return Return the corresponding FAutomationTestExcludelistEntry if excluded otherwise nullptr
+	*/
+	AUTOMATIONTEST_API const FAutomationTestExcludelistEntry* GetExcludeTestEntry(const FString& TestName, const TSet<FName>& RHI) const;
+
+	/**
+	 * Get the Entry object of the TestName if excluded
+	 *
+	 * @param TestName The full test name to be tested.
+	 * @param Platform The Platform the test is being evaluated against.
+	 * @param RHI The RHI the test is being evaluated against.
+	 * @return Return the corresponding FAutomationTestExcludelistEntry if excluded otherwise nullptr
+	*/
+	AUTOMATIONTEST_API const FAutomationTestExcludelistEntry* GetExcludeTestEntry(const FString& TestName, const FName& Platform, const TSet<FName>& RHI) const;
+
+	UE_DEPRECATED(5.4, "Please use GetConfigFilenameForEntry instead.")
+	AUTOMATIONTEST_API FString GetConfigFilename() const;
+
+	/**
+	 * Get the path to the config from which the exclusion was defined from
+	 *
+	 * @param Entry The FAutomationTestExcludelistEntry that is referenced.
+	 * @return Return the path to the corresponding config
+	*/
+	AUTOMATIONTEST_API FString GetConfigFilenameForEntry(const FAutomationTestExcludelistEntry& Entry) const;
+
+	/**
+	 * Get the path to the config from which the exclusion was defined from
+	 *
+	 * @param Entry The FAutomationTestExcludelistEntry that is referenced.
+	 * @param Platform The Platform the test is being evaluated against.
+	 * @return Return the path to the corresponding config
+	*/
+	AUTOMATIONTEST_API FString GetConfigFilenameForEntry(const FAutomationTestExcludelistEntry& Entry, const FName& Platform) const;
+
+	/** Save all the required configs */
+	AUTOMATIONTEST_API void SaveToConfigs();
+
+	UE_DEPRECATED(5.4, "Please use SaveToConfigs instead.")
+	AUTOMATIONTEST_API void SaveConfig() { SaveToConfigs(); }
 
-	AUTOMATIONTEST_API virtual void OverrideConfigSection(FString& SectionName) override;
+private:
+	/** Initiate the loading of the configs */
+	void Initialize();
 
-protected:
-	AUTOMATIONTEST_API virtual void PostInitProperties() override;
+	/** Load all possible platform configs */
+	void LoadPlatformConfigs();
 
-private:
-	AUTOMATIONTEST_API FString GetFullTestName(const FAutomationTestExcludelistEntry& ExcludelistEntry);
+	/** Populate the exclusion entries from the loaded configs */
+	void PopulateEntries();
 
-	UPROPERTY(EditDefaultsOnly, globalconfig, Category = AutomationTestExcludelist)
-	TArray<FAutomationTestExcludelistEntry> ExcludeTest;
-};
+	/** Exclusion entries */
+	TMap<FString, FAutomationTestExcludelistEntry> Entries;
 
+	UPROPERTY(Transient)
+	TObjectPtr<UAutomationTestExcludelistConfig> DefaultConfig;
+
+	UPROPERTY(Transient)
+	TMap<FName, TObjectPtr<UAutomationTestExcludelistConfig>> PlatformConfigs;
+};
diff --git a/Engine/Source/Runtime/AutomationTest/Public/AutomationTestPlatform.h b/Engine/Source/Runtime/AutomationTest/Public/AutomationTestPlatform.h
new file mode 100644
--- /dev/null
+++ b/Engine/Source/Runtime/AutomationTest/Public/AutomationTestPlatform.h
@@ -0,0 +1,44 @@
+// Copyright Epic Games, Inc. All Rights Reserved.
+
+#pragma once
+
+#include "CoreMinimal.h"
+#include "Templates/SubclassOf.h"
+
+#include "AutomationTestPlatform.generated.h"
+
+UCLASS(Abstract, MinimalAPI)
+class UAutomationTestPlatformSettings : public UObject
+{
+	GENERATED_BODY()
+
+public:
+	UAutomationTestPlatformSettings() {}
+
+	AUTOMATIONTEST_API static UAutomationTestPlatformSettings* Create(TSubclassOf<UAutomationTestPlatformSettings> SettingsClass, const FString PlatformName);
+
+	//~UObject interface
+	virtual void OverrideConfigSection(FString& SectionName) override;
+	virtual const TCHAR* GetConfigOverridePlatform() const override;
+	//~End of UObject interface
+
+	AUTOMATIONTEST_API void LoadConfig();
+	AUTOMATIONTEST_API FString GetConfigFilename() const;
+	AUTOMATIONTEST_API const FString& GetPlatformName() const;
+
+protected:
+	virtual void InitializeSettingsDefault() { }
+	virtual FString GetSectionName() { return StaticClass()->GetPathName(); }
+
+private:
+	FString PlatformName;
+};
+
+namespace AutomationTestPlatform
+{
+	/** Return the names of all available platforms */
+	AUTOMATIONTEST_API const TSet<FName>& GetAllAvailablePlatformNames();
+
+	/** Load and return target automation test settings collection of all platforms  */
+	AUTOMATIONTEST_API TArray<UAutomationTestPlatformSettings*>& GetAllPlatformsSettings(TSubclassOf<UAutomationTestPlatformSettings> SettingsClass);
+};
diff --git a/Engine/Source/Runtime/AutomationWorker/Private/AutomationWorkerModule.cpp b/Engine/Source/Runtime/AutomationWorker/Private/AutomationWorkerModule.cpp
--- a/Engine/Source/Runtime/AutomationWorker/Private/AutomationWorkerModule.cpp
+++ b/Engine/Source/Runtime/AutomationWorker/Private/AutomationWorkerModule.cpp
@@ -557,29 +557,35 @@ bool FAutomationWorkerModule::IsTestExcluded(const FString& InTestToRun, FString
 	check(nullptr != Excludelist);
 	
 	static const TSet<FName> RHI = GetRHIForAutomation();
-	if (Excludelist->IsTestExcluded(InTestToRun, RHI, &SkipReason, OutWarn))
+	static const FName Platform = FPlatformProperties::IniPlatformName();
+	if (Excludelist->IsTestExcluded(InTestToRun, Platform, RHI, &SkipReason, OutWarn))
 	{
 		if (OutReason)
 		{
 			(*OutReason) = (SkipReason.IsNone() ? TEXT("unknown reason") : SkipReason.ToString());
 			(*OutReason) += TEXT(" [config]");
 
-			if (Excludelist->GetFName().ToString().StartsWith("Default_"))
+			if (const FAutomationTestExcludelistEntry* Entry = Excludelist->GetExcludeTestEntry(InTestToRun))
 			{
-				// We can handle only CDO to correctly detect the filename
-				FString Filename = Excludelist->GetConfigFilename();
-				
+				if (!Entry->Platforms.IsEmpty())
+				{
+					(*OutReason) += TEXT(" [");
+					(*OutReason) += Platform.ToString();
+					(*OutReason) += TEXT("]");
+				}
+
+				FString Filename = Excludelist->GetConfigFilenameForEntry(*Entry, Platform);
+
 				if (!Filename.IsEmpty())
 				{
 					Filename = FPaths::ConvertRelativePathToFull(Filename);
-						FPaths::MakePlatformFilename(Filename);
-						*OutReason += TEXT(" [");
-						*OutReason += Filename;
-						// Using of line number 1 as a default value to get it working as a hyperlink
-						*OutReason += TEXT("(1)]");
+					FPaths::MakePlatformFilename(Filename);
+					*OutReason += TEXT(" [");
+					*OutReason += Filename;
+					// Using of line number 1 as a default value to get it working as a hyperlink
+					*OutReason += TEXT("(1)]");
 				}
 			}
-
 		}
 
 		return true;
