Test Exclude List Entries from default(all platforms) are overridden by platform specific entries

#rnx

diff --git a/Engine/Source/Runtime/AutomationTest/Private/AutomationTestExcludelist.cpp b/Engine/Source/Runtime/AutomationTest/Private/AutomationTestExcludelist.cpp
--- a/Engine/Source/Runtime/AutomationTest/Private/AutomationTestExcludelist.cpp
+++ b/Engine/Source/Runtime/AutomationTest/Private/AutomationTestExcludelist.cpp
@@ -166,6 +166,11 @@ void UAutomationTestExcludelist::PopulateEntries()
 		{
 			if (FAutomationTestExcludelistEntry* Entry = Entries.Find(PlatformEntry.FullTestName))
 			{
+				if (Entry->Platforms.IsEmpty())
+				{
+					continue;
+				}
+
 				Entry->Platforms.Add(Config.Key);
 				Entry->RHIs.Append(PlatformEntry.RHIs);
 			}
